Surface Effects Plugin initial commit

diff --git a/Engine/Plugins/Experimental/SurfaceEffects/Source/SurfaceEffects/Private/SurfaceEffects.cpp b/Engine/Plugins/Experimental/SurfaceEffects/Source/SurfaceEffects/Private/SurfaceEffects.cpp
new file mode 100644
--- /dev/null
+++ b/Engine/Plugins/Experimental/SurfaceEffects/Source/SurfaceEffects/Private/SurfaceEffects.cpp
@@ -0,0 +1,17 @@
+﻿// Copyright Epic Games, Inc. All Rights Reserved.
+
+#include "SurfaceEffects.h"
+
+#define LOCTEXT_NAMESPACE "FSurfaceEffectsModule"
+
+void FSurfaceEffectsModule::StartupModule()
+{
+}
+
+void FSurfaceEffectsModule::ShutdownModule()
+{
+}
+
+#undef LOCTEXT_NAMESPACE
+	
+IMPLEMENT_MODULE(FSurfaceEffectsModule, SurfaceEffects)
\ No newline at end of file
diff --git a/Engine/Plugins/Experimental/SurfaceEffects/Source/SurfaceEffects/Private/SurfaceEffectsSubsystem.cpp b/Engine/Plugins/Experimental/SurfaceEffects/Source/SurfaceEffects/Private/SurfaceEffectsSubsystem.cpp
new file mode 100644
--- /dev/null
+++ b/Engine/Plugins/Experimental/SurfaceEffects/Source/SurfaceEffects/Private/SurfaceEffectsSubsystem.cpp
@@ -0,0 +1,17 @@
+// Copyright Epic Games, Inc. All Rights Reserved.
+
+#include "SurfaceEffectsSubsystem.h"
+
+#include "SurfaceEffectsSettings.h"
+
+void USurfaceEffectsSubsystem::Initialize(FSubsystemCollectionBase& Collection)
+{
+	Super::Initialize(Collection);
+
+	const USurfaceEffectsSettings* MutableDefault = GetMutableDefault<USurfaceEffectsSettings>();
+
+	if(MutableDefault)
+	{
+		SurfaceEffectsData = LoadObject<UDataTable>(nullptr, *MutableDefault->SurfaceEffectsDataTable.ToString(), nullptr, LOAD_None, nullptr);
+	}
+}
diff --git a/Engine/Plugins/Experimental/SurfaceEffects/Source/SurfaceEffects/Public/SurfaceEffects.h b/Engine/Plugins/Experimental/SurfaceEffects/Source/SurfaceEffects/Public/SurfaceEffects.h
new file mode 100644
--- /dev/null
+++ b/Engine/Plugins/Experimental/SurfaceEffects/Source/SurfaceEffects/Public/SurfaceEffects.h
@@ -0,0 +1,13 @@
+﻿// Copyright Epic Games, Inc. All Rights Reserved.
+
+#pragma once
+
+#include "Modules/ModuleManager.h"
+
+class FSurfaceEffectsModule : public IModuleInterface
+{
+public:
+	/** IModuleInterface implementation */
+	virtual void StartupModule() override;
+	virtual void ShutdownModule() override;
+};
diff --git a/Engine/Plugins/Experimental/SurfaceEffects/Source/SurfaceEffects/Public/SurfaceEffectsSettings.h b/Engine/Plugins/Experimental/SurfaceEffects/Source/SurfaceEffects/Public/SurfaceEffectsSettings.h
new file mode 100644
--- /dev/null
+++ b/Engine/Plugins/Experimental/SurfaceEffects/Source/SurfaceEffects/Public/SurfaceEffectsSettings.h
@@ -0,0 +1,20 @@
+// Copyright Epic Games, Inc. All Rights Reserved.
+#pragma once
+
+#include "Engine/DeveloperSettings.h"
+
+#include "SurfaceEffectsSettings.generated.h"
+
+/**
+ * Surface Effects Settings.
+ */
+UCLASS(config=Engine, defaultconfig, meta=(DisplayName="Surface Effects"))
+class SURFACEEFFECTS_API USurfaceEffectsSettings : public UDeveloperSettings
+{
+	GENERATED_BODY()
+
+public:
+	/** List of data tables to load tags from */
+	UPROPERTY(config, EditAnywhere, Category = GameplayTags, meta = (AllowedClasses = "/Script/Engine.DataTable"))
+	FSoftObjectPath SurfaceEffectsDataTable;
+};
\ No newline at end of file
diff --git a/Engine/Plugins/Experimental/SurfaceEffects/Source/SurfaceEffects/Public/SurfaceEffectsSubsystem.h b/Engine/Plugins/Experimental/SurfaceEffects/Source/SurfaceEffects/Public/SurfaceEffectsSubsystem.h
new file mode 100644
--- /dev/null
+++ b/Engine/Plugins/Experimental/SurfaceEffects/Source/SurfaceEffects/Public/SurfaceEffectsSubsystem.h
@@ -0,0 +1,75 @@
+// Copyright Epic Games, Inc. All Rights Reserved.
+#pragma once
+
+#include "Engine/DataAsset.h"
+#include "Engine/DataTable.h"
+#include "Subsystems/GameInstanceSubsystem.h"
+
+#include "SurfaceEffectsSubsystem.generated.h"
+
+/**
+ * Base data asset used to store what conditions result in a specific surface being returned
+ */
+UCLASS(Abstract)
+class SURFACEEFFECTS_API USurfaceEffectRule : public UDataAsset
+{
+	GENERATED_BODY()
+};
+
+/**
+ * Data Table Row that effectively wraps the Surface Effect Rule
+ */
+USTRUCT()
+struct FSurfaceEffectTableRow : public FTableRowBase
+{
+	GENERATED_BODY()
+	
+	UPROPERTY(EditAnywhere , Category="Surface Effect")
+	TObjectPtr<USurfaceEffectRule> Rule;
+};
+
+/**
+ * Base context for determining which enum value to return based on a certain rule
+ * @tparam TEnum An Enum class with the UENUM Specifier
+ */
+template <typename TEnum>
+struct TSurfaceEffectContext
+{
+	virtual ~TSurfaceEffectContext() = default;
+
+	explicit TSurfaceEffectContext()
+	{
+	}
+
+	virtual void GetSurface(TEnum& OutSurface, USurfaceEffectRule* Rule) = 0;
+};
+
+/**
+ * A system for handling various surface enums based on contexts
+ */
+UCLASS(MinimalAPI)
+class USurfaceEffectsSubsystem : public UGameInstanceSubsystem
+{
+	GENERATED_BODY()
+
+public:
+	virtual void Initialize(FSubsystemCollectionBase& Collection) override;
+
+	template <typename TEnum>
+	void GetSurface(TEnum& OutSurface, TSurfaceEffectContext<TEnum>& Context)
+	{
+		const UEnum* EnumClass = StaticEnum<TEnum>();
+		if(SurfaceEffectsData)
+		{
+			FSurfaceEffectTableRow* Row = SurfaceEffectsData->FindRow<FSurfaceEffectTableRow>(EnumClass->GetFName(), ANSI_TO_TCHAR(__FUNCTION__));
+			Context.GetSurface(OutSurface, Row->Rule);
+		}
+	}
+
+private:
+	/*
+	 * We store the UEnum name as the row name in a data table to get the rule associated with that surface enum
+	 */
+	UPROPERTY()
+	TObjectPtr<UDataTable> SurfaceEffectsData;
+};
diff --git a/Engine/Plugins/Experimental/SurfaceEffects/Source/SurfaceEffects/SurfaceEffects.Build.cs b/Engine/Plugins/Experimental/SurfaceEffects/Source/SurfaceEffects/SurfaceEffects.Build.cs
new file mode 100644
--- /dev/null
+++ b/Engine/Plugins/Experimental/SurfaceEffects/Source/SurfaceEffects/SurfaceEffects.Build.cs
@@ -0,0 +1,35 @@
+﻿// Copyright Epic Games, Inc. All Rights Reserved.
+
+using UnrealBuildTool;
+
+public class SurfaceEffects : ModuleRules
+{
+	public SurfaceEffects(ReadOnlyTargetRules Target) : base(Target)
+	{
+		PCHUsage = ModuleRules.PCHUsageMode.UseExplicitOrSharedPCHs;
+		
+		PublicDependencyModuleNames.AddRange(
+			new string[]
+			{
+				"Core",
+			}
+			);
+			
+		
+		PrivateDependencyModuleNames.AddRange(
+			new string[]
+			{
+				"CoreUObject",
+				"DeveloperSettings",
+				"Engine",
+			}
+			);
+		
+		
+		DynamicallyLoadedModuleNames.AddRange(
+			new string[]
+			{
+			}
+			);
+	}
+}
diff --git a/Engine/Plugins/Experimental/SurfaceEffects/SurfaceEffects.uplugin b/Engine/Plugins/Experimental/SurfaceEffects/SurfaceEffects.uplugin
new file mode 100644
--- /dev/null
+++ b/Engine/Plugins/Experimental/SurfaceEffects/SurfaceEffects.uplugin
@@ -0,0 +1,23 @@
+﻿{
+	"FileVersion": 3,
+	"Version": 1,
+	"VersionName": "1.0",
+	"FriendlyName": "SurfaceEffects",
+	"Description": "A flexible, context-driven surface system",
+	"Category": "Other",
+	"CreatedBy": "Epic Games, Inc.",
+	"CreatedByURL": "https://epicgames.com",
+	"DocsURL": "",
+	"MarketplaceURL": "",
+	"CanContainContent": false,
+	"IsBetaVersion": true,
+	"IsExperimentalVersion": true,
+	"Installed": false,
+	"Modules": [
+		{
+			"Name": "SurfaceEffects",
+			"Type": "Runtime",
+			"LoadingPhase": "Default"
+		}
+	]
+}
\ No newline at end of file
