First implementation of Surface Effects MVP - Footsteps

#virtualized

diff --git a/Engine/Plugins/Experimental/SurfaceEffects/Source/SurfaceEffects/Private/SurfaceEffectsSubsystem.cpp b/Engine/Plugins/Experimental/SurfaceEffects/Source/SurfaceEffects/Private/SurfaceEffectsSubsystem.cpp
--- a/Engine/Plugins/Experimental/SurfaceEffects/Source/SurfaceEffects/Private/SurfaceEffectsSubsystem.cpp
+++ b/Engine/Plugins/Experimental/SurfaceEffects/Source/SurfaceEffects/Private/SurfaceEffectsSubsystem.cpp
@@ -2,16 +2,34 @@
 
 #include "SurfaceEffectsSubsystem.h"
 
+#include "HAL/IConsoleManager.h"
 #include "SurfaceEffectsSettings.h"
 
+#include UE_INLINE_GENERATED_CPP_BY_NAME(SurfaceEffectsSubsystem)
+
+namespace SurfaceEffectConsoleVariables
+{
+	int32 bEnabled = 1;
+	FAutoConsoleVariableRef CVarEnabled(
+	TEXT("SurfaceEffects.Enabled"), bEnabled,
+	TEXT("Enables the Surface Effects System.\n")
+		TEXT("0: Disabled, 1: Enabled"),
+	ECVF_Default);
+} // namespace SurfaceEffectConsoleVariables
+
 void USurfaceEffectsSubsystem::Initialize(FSubsystemCollectionBase& Collection)
 {
 	Super::Initialize(Collection);
 
-	const USurfaceEffectsSettings* MutableDefault = GetMutableDefault<USurfaceEffectsSettings>();
+	const USurfaceEffectsSettings* Settings = GetDefault<USurfaceEffectsSettings>();
 
-	if(MutableDefault)
+	if(Settings)
 	{
-		SurfaceEffectsData = LoadObject<UDataTable>(nullptr, *MutableDefault->SurfaceEffectsDataTable.ToString(), nullptr, LOAD_None, nullptr);
+		SurfaceEffectsData = LoadObject<UDataTable>(nullptr, *Settings->SurfaceEffectsDataTable.ToString(), nullptr, LOAD_None, nullptr);
 	}
 }
+
+bool USurfaceEffectsSubsystem::IsEnabled()
+{
+	return SurfaceEffectConsoleVariables::bEnabled != 0;
+}
diff --git a/Engine/Plugins/Experimental/SurfaceEffects/Source/SurfaceEffects/Public/SurfaceEffectsSettings.h b/Engine/Plugins/Experimental/SurfaceEffects/Source/SurfaceEffects/Public/SurfaceEffectsSettings.h
--- a/Engine/Plugins/Experimental/SurfaceEffects/Source/SurfaceEffects/Public/SurfaceEffectsSettings.h
+++ b/Engine/Plugins/Experimental/SurfaceEffects/Source/SurfaceEffects/Public/SurfaceEffectsSettings.h
@@ -5,16 +5,15 @@
 
 #include "SurfaceEffectsSettings.generated.h"
 
-/**
- * Surface Effects Settings.
- */
 UCLASS(config=Engine, defaultconfig, meta=(DisplayName="Surface Effects"))
 class SURFACEEFFECTS_API USurfaceEffectsSettings : public UDeveloperSettings
 {
 	GENERATED_BODY()
 
 public:
-	/** List of data tables to load tags from */
-	UPROPERTY(config, EditAnywhere, Category = GameplayTags, meta = (AllowedClasses = "/Script/Engine.DataTable"))
+	/**
+	 * Data table for storing surface effect rules @see FSurfaceEffectTableRow
+	 */
+	UPROPERTY(config, EditAnywhere, Category = GameplayTags, meta = (AllowedClasses = "/Script/Engine.DataTable", RowType = "/Script/SurfaceEffects.SurfaceEffectTableRow"))
 	FSoftObjectPath SurfaceEffectsDataTable;
 };
\ No newline at end of file
diff --git a/Engine/Plugins/Experimental/SurfaceEffects/Source/SurfaceEffects/Public/SurfaceEffectsSubsystem.h b/Engine/Plugins/Experimental/SurfaceEffects/Source/SurfaceEffects/Public/SurfaceEffectsSubsystem.h
--- a/Engine/Plugins/Experimental/SurfaceEffects/Source/SurfaceEffects/Public/SurfaceEffectsSubsystem.h
+++ b/Engine/Plugins/Experimental/SurfaceEffects/Source/SurfaceEffects/Public/SurfaceEffectsSubsystem.h
@@ -7,41 +7,66 @@
 
 #include "SurfaceEffectsSubsystem.generated.h"
 
+enum EPhysicalSurface : int;
+
 /**
- * Base data asset used to store what conditions result in a specific surface being returned
+ * Base context for determining which enum value to return based on a certain rule.
+ * We assume most surface interactions will want to use EPhysicalSurface as part of the context.
  */
-UCLASS(Abstract)
-class SURFACEEFFECTS_API USurfaceEffectRule : public UDataAsset
+struct FSurfaceEffectContextBase
 {
-	GENERATED_BODY()
+	explicit FSurfaceEffectContextBase(const EPhysicalSurface PhysicalSurface)
+		: PhysicalSurface(PhysicalSurface)
+	{
+	}
+	
+	EPhysicalSurface PhysicalSurface;
 };
 
 /**
- * Data Table Row that effectively wraps the Surface Effect Rule
+ *	TSurfaceEffectResult - results of a GetSurface call
  */
-USTRUCT()
-struct FSurfaceEffectTableRow : public FTableRowBase
+template <typename TEnum>
+struct TSurfaceEffectResult
 {
-	GENERATED_BODY()
+	TEnum OutSurface;
 	
-	UPROPERTY(EditAnywhere , Category="Surface Effect")
-	TObjectPtr<USurfaceEffectRule> Rule;
+	/** Set to true if we successfully got a surface */
+	bool bSuccess = false;
 };
 
+
 /**
- * Base context for determining which enum value to return based on a certain rule
- * @tparam TEnum An Enum class with the UENUM Specifier
+ * Base data asset used to store what conditions result in a specific surface being returned
  */
-template <typename TEnum>
-struct TSurfaceEffectContext
+UCLASS(Abstract)
+class SURFACEEFFECTS_API USurfaceEffectRule : public UDataAsset
 {
-	virtual ~TSurfaceEffectContext() = default;
+	GENERATED_BODY()
 
-	explicit TSurfaceEffectContext()
+public:
+	/**
+	 * 
+	 * @param OutSurfaceValue The enum value to be used when casting back to the enum class
+	 * @param Context Context data used by the rule to determine what value to return
+	 * @return Returns true if a valid enum value is outputted. Otherwise, return false
+	 */
+	virtual bool GetSurface(uint8& OutSurfaceValue, const FSurfaceEffectContextBase& Context)
 	{
+		return false;
 	}
+};
 
-	virtual void GetSurface(TEnum& OutSurface, USurfaceEffectRule* Rule) = 0;
+/**
+ * Data Table Row that effectively wraps the Surface Effect Rule
+ */
+USTRUCT()
+struct SURFACEEFFECTS_API FSurfaceEffectTableRow : public FTableRowBase
+{
+	GENERATED_BODY()
+	
+	UPROPERTY(EditAnywhere , Category="Surface Effect")
+	TObjectPtr<USurfaceEffectRule> Rule;
 };
 
 /**
@@ -56,14 +81,31 @@ public:
 	virtual void Initialize(FSubsystemCollectionBase& Collection) override;
 
 	template <typename TEnum>
-	void GetSurface(TEnum& OutSurface, TSurfaceEffectContext<TEnum>& Context)
+	TSurfaceEffectResult<TEnum> GetSurface(const FSurfaceEffectContextBase& Context)
 	{
-		const UEnum* EnumClass = StaticEnum<TEnum>();
-		if(SurfaceEffectsData)
+		TSurfaceEffectResult<TEnum> Result;
+		if(IsEnabled() && SurfaceEffectsData)
 		{
-			FSurfaceEffectTableRow* Row = SurfaceEffectsData->FindRow<FSurfaceEffectTableRow>(EnumClass->GetFName(), ANSI_TO_TCHAR(__FUNCTION__));
-			Context.GetSurface(OutSurface, Row->Rule);
+			static_assert(TIsEnum<TEnum>::Value, "Should only call this with enum types");
+			const UEnum* EnumClass = StaticEnum<TEnum>();
+			if(EnumClass)
+			{
+				const FString ContextString(TEXT("USurfaceEffectsSubsystem::GetSurface"));
+				const FSurfaceEffectTableRow* Row = SurfaceEffectsData->FindRow<FSurfaceEffectTableRow>(EnumClass->GetFName(), ContextString);
+				if(Row && Row->Rule)
+				{
+					const UEnum* Enum = StaticEnum<TEnum>();
+					uint8 OutSurfaceIndex;
+					Result.bSuccess = Row->Rule->GetSurface(OutSurfaceIndex, Context);
+					Result.bSuccess &= OutSurfaceIndex < Enum->GetMaxEnumValue();
+					if(Result.bSuccess)
+					{
+						Result.OutSurface = static_cast<TEnum>(OutSurfaceIndex);
+					}
+				}
+			}
 		}
+		return Result;
 	}
 
 private:
@@ -72,4 +114,6 @@ private:
 	 */
 	UPROPERTY()
 	TObjectPtr<UDataTable> SurfaceEffectsData;
+
+	SURFACEEFFECTS_API static bool IsEnabled();
 };
diff --git a/Engine/Plugins/Experimental/SurfaceEffects/Source/SurfaceEffects/SurfaceEffects.Build.cs b/Engine/Plugins/Experimental/SurfaceEffects/Source/SurfaceEffects/SurfaceEffects.Build.cs
--- a/Engine/Plugins/Experimental/SurfaceEffects/Source/SurfaceEffects/SurfaceEffects.Build.cs
+++ b/Engine/Plugins/Experimental/SurfaceEffects/Source/SurfaceEffects/SurfaceEffects.Build.cs
@@ -12,6 +12,7 @@ public class SurfaceEffects : ModuleRules
 			new string[]
 			{
 				"Core",
+				"DeveloperSettings",
 			}
 			);
 			
@@ -20,16 +21,8 @@ public class SurfaceEffects : ModuleRules
 			new string[]
 			{
 				"CoreUObject",
-				"DeveloperSettings",
 				"Engine",
 			}
 			);
-		
-		
-		DynamicallyLoadedModuleNames.AddRange(
-			new string[]
-			{
-			}
-			);
 	}
 }
