[Backout] - CL30970339
Original CL Desc
-----------------------------------------------------------------
First implementation of Surface Effects MVP - Footsteps

#virtualized

diff --git a/Engine/Plugins/Experimental/SurfaceEffects/Source/SurfaceEffects/Private/SurfaceEffectsSubsystem.cpp b/Engine/Plugins/Experimental/SurfaceEffects/Source/SurfaceEffects/Private/SurfaceEffectsSubsystem.cpp
--- a/Engine/Plugins/Experimental/SurfaceEffects/Source/SurfaceEffects/Private/SurfaceEffectsSubsystem.cpp
+++ b/Engine/Plugins/Experimental/SurfaceEffects/Source/SurfaceEffects/Private/SurfaceEffectsSubsystem.cpp
@@ -2,34 +2,16 @@
 
 #include "SurfaceEffectsSubsystem.h"
 
-#include "HAL/IConsoleManager.h"
 #include "SurfaceEffectsSettings.h"
 
-#include UE_INLINE_GENERATED_CPP_BY_NAME(SurfaceEffectsSubsystem)
-
-namespace SurfaceEffectConsoleVariables
-{
-	int32 bEnabled = 1;
-	FAutoConsoleVariableRef CVarEnabled(
-	TEXT("SurfaceEffects.Enabled"), bEnabled,
-	TEXT("Enables the Surface Effects System.\n")
-		TEXT("0: Disabled, 1: Enabled"),
-	ECVF_Default);
-} // namespace SurfaceEffectConsoleVariables
-
 void USurfaceEffectsSubsystem::Initialize(FSubsystemCollectionBase& Collection)
 {
 	Super::Initialize(Collection);
 
-	const USurfaceEffectsSettings* Settings = GetDefault<USurfaceEffectsSettings>();
+	const USurfaceEffectsSettings* MutableDefault = GetMutableDefault<USurfaceEffectsSettings>();
 
-	if(Settings)
+	if(MutableDefault)
 	{
-		SurfaceEffectsData = LoadObject<UDataTable>(nullptr, *Settings->SurfaceEffectsDataTable.ToString(), nullptr, LOAD_None, nullptr);
+		SurfaceEffectsData = LoadObject<UDataTable>(nullptr, *MutableDefault->SurfaceEffectsDataTable.ToString(), nullptr, LOAD_None, nullptr);
 	}
 }
-
-bool USurfaceEffectsSubsystem::IsEnabled()
-{
-	return SurfaceEffectConsoleVariables::bEnabled != 0;
-}
diff --git a/Engine/Plugins/Experimental/SurfaceEffects/Source/SurfaceEffects/Public/SurfaceEffectsSettings.h b/Engine/Plugins/Experimental/SurfaceEffects/Source/SurfaceEffects/Public/SurfaceEffectsSettings.h
--- a/Engine/Plugins/Experimental/SurfaceEffects/Source/SurfaceEffects/Public/SurfaceEffectsSettings.h
+++ b/Engine/Plugins/Experimental/SurfaceEffects/Source/SurfaceEffects/Public/SurfaceEffectsSettings.h
@@ -5,15 +5,16 @@
 
 #include "SurfaceEffectsSettings.generated.h"
 
+/**
+ * Surface Effects Settings.
+ */
 UCLASS(config=Engine, defaultconfig, meta=(DisplayName="Surface Effects"))
 class SURFACEEFFECTS_API USurfaceEffectsSettings : public UDeveloperSettings
 {
 	GENERATED_BODY()
 
 public:
-	/**
-	 * Data table for storing surface effect rules @see FSurfaceEffectTableRow
-	 */
-	UPROPERTY(config, EditAnywhere, Category = GameplayTags, meta = (AllowedClasses = "/Script/Engine.DataTable", RowType = "/Script/SurfaceEffects.SurfaceEffectTableRow"))
+	/** List of data tables to load tags from */
+	UPROPERTY(config, EditAnywhere, Category = GameplayTags, meta = (AllowedClasses = "/Script/Engine.DataTable"))
 	FSoftObjectPath SurfaceEffectsDataTable;
 };
\ No newline at end of file
diff --git a/Engine/Plugins/Experimental/SurfaceEffects/Source/SurfaceEffects/Public/SurfaceEffectsSubsystem.h b/Engine/Plugins/Experimental/SurfaceEffects/Source/SurfaceEffects/Public/SurfaceEffectsSubsystem.h
--- a/Engine/Plugins/Experimental/SurfaceEffects/Source/SurfaceEffects/Public/SurfaceEffectsSubsystem.h
+++ b/Engine/Plugins/Experimental/SurfaceEffects/Source/SurfaceEffects/Public/SurfaceEffectsSubsystem.h
@@ -7,35 +7,6 @@
 
 #include "SurfaceEffectsSubsystem.generated.h"
 
-enum EPhysicalSurface : int;
-
-/**
- * Base context for determining which enum value to return based on a certain rule.
- * We assume most surface interactions will want to use EPhysicalSurface as part of the context.
- */
-struct FSurfaceEffectContextBase
-{
-	explicit FSurfaceEffectContextBase(const EPhysicalSurface PhysicalSurface)
-		: PhysicalSurface(PhysicalSurface)
-	{
-	}
-	
-	EPhysicalSurface PhysicalSurface;
-};
-
-/**
- *	TSurfaceEffectResult - results of a GetSurface call
- */
-template <typename TEnum>
-struct TSurfaceEffectResult
-{
-	TEnum OutSurface;
-	
-	/** Set to true if we successfully got a surface */
-	bool bSuccess = false;
-};
-
-
 /**
  * Base data asset used to store what conditions result in a specific surface being returned
  */
@@ -43,25 +14,13 @@ UCLASS(Abstract)
 class SURFACEEFFECTS_API USurfaceEffectRule : public UDataAsset
 {
 	GENERATED_BODY()
-
-public:
-	/**
-	 * 
-	 * @param OutSurfaceValue The enum value to be used when casting back to the enum class
-	 * @param Context Context data used by the rule to determine what value to return
-	 * @return Returns true if a valid enum value is outputted. Otherwise, return false
-	 */
-	virtual bool GetSurface(uint8& OutSurfaceValue, const FSurfaceEffectContextBase& Context)
-	{
-		return false;
-	}
 };
 
 /**
  * Data Table Row that effectively wraps the Surface Effect Rule
  */
 USTRUCT()
-struct SURFACEEFFECTS_API FSurfaceEffectTableRow : public FTableRowBase
+struct FSurfaceEffectTableRow : public FTableRowBase
 {
 	GENERATED_BODY()
 	
@@ -69,6 +28,22 @@ struct SURFACEEFFECTS_API FSurfaceEffectTableRow : public FTableRowBase
 	TObjectPtr<USurfaceEffectRule> Rule;
 };
 
+/**
+ * Base context for determining which enum value to return based on a certain rule
+ * @tparam TEnum An Enum class with the UENUM Specifier
+ */
+template <typename TEnum>
+struct TSurfaceEffectContext
+{
+	virtual ~TSurfaceEffectContext() = default;
+
+	explicit TSurfaceEffectContext()
+	{
+	}
+
+	virtual void GetSurface(TEnum& OutSurface, USurfaceEffectRule* Rule) = 0;
+};
+
 /**
  * A system for handling various surface enums based on contexts
  */
@@ -81,31 +56,14 @@ public:
 	virtual void Initialize(FSubsystemCollectionBase& Collection) override;
 
 	template <typename TEnum>
-	TSurfaceEffectResult<TEnum> GetSurface(const FSurfaceEffectContextBase& Context)
+	void GetSurface(TEnum& OutSurface, TSurfaceEffectContext<TEnum>& Context)
 	{
-		TSurfaceEffectResult<TEnum> Result;
-		if(IsEnabled() && SurfaceEffectsData)
+		const UEnum* EnumClass = StaticEnum<TEnum>();
+		if(SurfaceEffectsData)
 		{
-			static_assert(TIsEnum<TEnum>::Value, "Should only call this with enum types");
-			const UEnum* EnumClass = StaticEnum<TEnum>();
-			if(EnumClass)
-			{
-				const FString ContextString(TEXT("USurfaceEffectsSubsystem::GetSurface"));
-				const FSurfaceEffectTableRow* Row = SurfaceEffectsData->FindRow<FSurfaceEffectTableRow>(EnumClass->GetFName(), ContextString);
-				if(Row && Row->Rule)
-				{
-					const UEnum* Enum = StaticEnum<TEnum>();
-					uint8 OutSurfaceIndex;
-					Result.bSuccess = Row->Rule->GetSurface(OutSurfaceIndex, Context);
-					Result.bSuccess &= OutSurfaceIndex < Enum->GetMaxEnumValue();
-					if(Result.bSuccess)
-					{
-						Result.OutSurface = static_cast<TEnum>(OutSurfaceIndex);
-					}
-				}
-			}
+			FSurfaceEffectTableRow* Row = SurfaceEffectsData->FindRow<FSurfaceEffectTableRow>(EnumClass->GetFName(), ANSI_TO_TCHAR(__FUNCTION__));
+			Context.GetSurface(OutSurface, Row->Rule);
 		}
-		return Result;
 	}
 
 private:
@@ -114,6 +72,4 @@ private:
 	 */
 	UPROPERTY()
 	TObjectPtr<UDataTable> SurfaceEffectsData;
-
-	SURFACEEFFECTS_API static bool IsEnabled();
 };
diff --git a/Engine/Plugins/Experimental/SurfaceEffects/Source/SurfaceEffects/SurfaceEffects.Build.cs b/Engine/Plugins/Experimental/SurfaceEffects/Source/SurfaceEffects/SurfaceEffects.Build.cs
--- a/Engine/Plugins/Experimental/SurfaceEffects/Source/SurfaceEffects/SurfaceEffects.Build.cs
+++ b/Engine/Plugins/Experimental/SurfaceEffects/Source/SurfaceEffects/SurfaceEffects.Build.cs
@@ -12,7 +12,6 @@ public class SurfaceEffects : ModuleRules
 			new string[]
 			{
 				"Core",
-				"DeveloperSettings",
 			}
 			);
 			
@@ -21,8 +20,16 @@ public class SurfaceEffects : ModuleRules
 			new string[]
 			{
 				"CoreUObject",
+				"DeveloperSettings",
 				"Engine",
 			}
 			);
+		
+		
+		DynamicallyLoadedModuleNames.AddRange(
+			new string[]
+			{
+			}
+			);
 	}
 }
